USD: Move MDL schema translator to dedicated plugin that can be independently disabled.

diff --git a/Engine/Plugins/Importers/USDImporter/Source/USDSchemas/Private/Custom/MDLUSDShadeMaterialTranslator.cpp b/Engine/Plugins/Importers/USDImporter/Source/USDSchemas/Private/Custom/MDLUSDShadeMaterialTranslator.cpp
deleted file mode 100644
diff --git a/Engine/Plugins/Importers/USDImporter/Source/USDSchemas/Private/Custom/MDLUSDShadeMaterialTranslator.h b/Engine/Plugins/Importers/USDImporter/Source/USDSchemas/Private/Custom/MDLUSDShadeMaterialTranslator.h
deleted file mode 100644
diff --git a/Engine/Plugins/Importers/USDImporter/Source/USDSchemas/Private/Custom/MaterialXUSDShadeMaterialTranslator.h b/Engine/Plugins/Importers/USDImporter/Source/USDSchemas/Private/Custom/MaterialXUSDShadeMaterialTranslator.h
deleted file mode 100644
diff --git a/Engine/Plugins/Importers/USDImporter/Source/USDSchemas/Private/USDSchemasModule.cpp b/Engine/Plugins/Importers/USDImporter/Source/USDSchemas/Private/USDSchemasModule.cpp
--- a/Engine/Plugins/Importers/USDImporter/Source/USDSchemas/Private/USDSchemasModule.cpp
+++ b/Engine/Plugins/Importers/USDImporter/Source/USDSchemas/Private/USDSchemasModule.cpp
@@ -2,20 +2,15 @@
 
 #include "USDSchemasModule.h"
 
-#include "USDMemory.h"
 #include "Objects/USDSchemaTranslator.h"
+#include "USDMemory.h"
 #include "USDTypesConversion.h"
 
-#if WITH_EDITOR
-#include "MDLImporterModule.h"
-#endif	  // #if WITH_EDITOR
-
 #include "AnalyticsEventAttribute.h"
 #include "Modules/ModuleManager.h"
 
 #if USE_USD_SDK
 #include "Custom/MaterialXUSDShadeMaterialTranslator.h"
-#include "Custom/MDLUSDShadeMaterialTranslator.h"
 #include "USDGeomCameraTranslator.h"
 #include "USDGeometryCacheTranslator.h"
 #include "USDGeomMeshTranslator.h"
@@ -81,12 +76,6 @@ public:
 				 // It doesn't seem possible to create SoundWave assets at runtime at the moment, for whatever reason
 				 Registry.Register<FUsdMediaSpatialAudioTranslator>(TEXT("UsdMediaSpatialAudio"))}
 			);
-
-			if (IMDLImporterModule* MDLImporterModule = FModuleManager::Get().LoadModulePtr<IMDLImporterModule>(TEXT("MDLImporter")))
-			{
-				UsdUnreal::MaterialUtils::RegisterRenderContext(UnrealIdentifiers::MdlRenderContext);
-				TranslatorHandles.Add(Registry.Register<FMdlUsdShadeMaterialTranslator>(TEXT("UsdShadeMaterial")));
-			}
 		}
 #endif	  // WITH_EDITOR
 
@@ -105,7 +94,6 @@ public:
 		}
 
 #if USE_USD_SDK && WITH_EDITOR
-		UsdUnreal::MaterialUtils::UnregisterRenderContext(UnrealIdentifiers::MdlRenderContext);
 		UsdUnreal::MaterialUtils::UnregisterRenderContext(UnrealIdentifiers::MaterialXRenderContext);
 #endif	  // WITH_EDITOR
 	}
diff --git a/Engine/Plugins/Importers/USDImporter/Source/USDSchemas/Public/Custom/MaterialXUSDShadeMaterialTranslator.h b/Engine/Plugins/Importers/USDImporter/Source/USDSchemas/Public/Custom/MaterialXUSDShadeMaterialTranslator.h
new file mode 100644
--- /dev/null
+++ b/Engine/Plugins/Importers/USDImporter/Source/USDSchemas/Public/Custom/MaterialXUSDShadeMaterialTranslator.h
@@ -0,0 +1,23 @@
+// Copyright Epic Games, Inc. All Rights Reserved.
+
+#pragma once
+
+#include "USDShadeMaterialTranslator.h"
+
+#if USE_USD_SDK
+
+#include "USDIncludesStart.h"
+#include "pxr/pxr.h"
+#include "USDIncludesEnd.h"
+
+class USDSCHEMAS_API FMaterialXUsdShadeMaterialTranslator : public FUsdShadeMaterialTranslator
+{
+	using Super = FUsdShadeMaterialTranslator;
+
+public:
+	using FUsdShadeMaterialTranslator::FUsdShadeMaterialTranslator;
+
+	virtual void CreateAssets() override;
+};
+
+#endif	  // #if USE_USD_SDK
diff --git a/Engine/Plugins/Importers/USDImporter/Source/USDSchemas/USDSchemas.Build.cs b/Engine/Plugins/Importers/USDImporter/Source/USDSchemas/USDSchemas.Build.cs
--- a/Engine/Plugins/Importers/USDImporter/Source/USDSchemas/USDSchemas.Build.cs
+++ b/Engine/Plugins/Importers/USDImporter/Source/USDSchemas/USDSchemas.Build.cs
@@ -60,7 +60,6 @@ namespace UnrealBuildTool.Rules
 						"Kismet",
 						"LiveLinkGraphNode",
 						"MaterialEditor",
-						"MDLImporter",
 						"MeshUtilities",
 						"PhysicsUtilities", // For generating UPhysicsAssets for SkeletalMeshes and ConvexDecompTool
 						"PropertyEditor",
diff --git a/Engine/Plugins/Importers/USDImporter/USDImporter.uplugin b/Engine/Plugins/Importers/USDImporter/USDImporter.uplugin
--- a/Engine/Plugins/Importers/USDImporter/USDImporter.uplugin
+++ b/Engine/Plugins/Importers/USDImporter/USDImporter.uplugin
@@ -126,15 +126,6 @@
 			"Name": "LiveLink",
 			"Enabled": true
 		},
-		{
-			"Name": "MDLImporter",
-			"Enabled": true,
-			"PlatformAllowList": [
-				"Win64",
-				"Mac",
-				"Linux"
-			]
-		},
 		{
 			"Name": "PythonScriptPlugin",
 			"Enabled": true
diff --git a/Engine/Plugins/Importers/USDImporterMDL/Source/USDImporterMDL/Private/MDLUSDShadeMaterialTranslator.cpp b/Engine/Plugins/Importers/USDImporterMDL/Source/USDImporterMDL/Private/MDLUSDShadeMaterialTranslator.cpp
new file mode 100644
--- /dev/null
+++ b/Engine/Plugins/Importers/USDImporterMDL/Source/USDImporterMDL/Private/MDLUSDShadeMaterialTranslator.cpp
@@ -0,0 +1,339 @@
+// Copyright Epic Games, Inc. All Rights Reserved.
+
+#include "MDLUSDShadeMaterialTranslator.h"
+#include "Engine/Level.h"
+
+#if USE_USD_SDK && WITH_EDITOR
+
+#include "MDLUSDLog.h"
+#include "USDAssetCache.h"
+#include "USDAssetImportData.h"
+#include "USDConversionUtils.h"
+#include "USDErrorUtils.h"
+#include "USDMemory.h"
+#include "USDShadeConversion.h"
+#include "USDTypesConversion.h"
+
+#include "UsdWrappers/SdfLayer.h"
+#include "UsdWrappers/SdfPath.h"
+#include "UsdWrappers/UsdPrim.h"
+
+#include "Engine/Texture.h"
+#include "Materials/Material.h"
+#include "Materials/MaterialInstance.h"
+#include "Materials/MaterialInstanceConstant.h"
+#include "MaterialShared.h"
+#include "MDLImporterOptions.h"
+#include "MDLMaterialImporter.h"
+#include "Misc/Paths.h"
+#include "UObject/StrongObjectPtr.h"
+
+#include "USDIncludesStart.h"
+#include "pxr/usd/sdf/assetPath.h"
+#include "pxr/usd/usdShade/material.h"
+#include "USDIncludesEnd.h"
+
+#define LOCTEXT_NAMESPACE "MDLUSDShadeMaterialTranslator"
+
+PRAGMA_DISABLE_DEPRECATION_WARNINGS
+FName FMdlUsdShadeMaterialTranslator::MdlRenderContext = TEXT("mdl");
+PRAGMA_ENABLE_DEPRECATION_WARNINGS
+
+namespace UE::MDLShadeTranslator::Private
+{
+	void NotifyIfMaterialNeedsVirtualTextures(UMaterialInterface* MaterialInterface)
+	{
+		if (UMaterial* Material = Cast<UMaterial>(MaterialInterface))
+		{
+			TArray<UTexture*> UsedTextures;
+			const bool bAllQualityLevels = true;
+			const bool bAllFeatureLevels = true;
+			Material->GetUsedTextures(UsedTextures, EMaterialQualityLevel::High, bAllQualityLevels, ERHIFeatureLevel::SM5, bAllFeatureLevels);
+
+			for (UTexture* UsedTexture : UsedTextures)
+			{
+				UsdUtils::NotifyIfVirtualTexturesNeeded(UsedTexture);
+			}
+		}
+		else if (UMaterialInstance* MaterialInstance = Cast<UMaterialInstance>(MaterialInterface))
+		{
+			for (const FTextureParameterValue& TextureValue : MaterialInstance->TextureParameterValues)
+			{
+				if (UTexture* Texture = TextureValue.ParameterValue)
+				{
+					UsdUtils::NotifyIfVirtualTexturesNeeded(Texture);
+				}
+			}
+		}
+	}
+}	 // namespace UE::MDLShadeTranslator::Private
+
+void FMdlUsdShadeMaterialTranslator::CreateAssets()
+{
+	// MDL USD Schema:
+	//   info:mdl:sourceAsset -> Path to the MDL file
+	//   info:mdl:sourceAsset:subIdentifier -> Name of the material in the MDL file
+	//   inputs -> material parameters
+
+	if (Context->RenderContext != UnrealIdentifiers::MdlRenderContext)
+	{
+		Super::CreateAssets();
+		return;
+	}
+
+	if (Context->bTranslateOnlyUsedMaterials && Context->UsdInfoCache)
+	{
+		if (!Context->UsdInfoCache->IsMaterialUsed(PrimPath))
+		{
+			UE_LOG(
+				LogUsdMdl,
+				Verbose,
+				TEXT("Skipping creating assets for material prim '%s' as it is not currently bound by any prim."),
+				*PrimPath.GetString()
+			);
+			return;
+		}
+	}
+
+	FScopedUsdAllocs UsdAllocs;
+
+	pxr::UsdShadeMaterial ShadeMaterial(GetPrim());
+
+	if (!ShadeMaterial)
+	{
+		return;
+	}
+
+	const static pxr::TfToken MdlToken = UnrealToUsd::ConvertToken(*UnrealIdentifiers::MdlRenderContext.ToString()).Get();
+
+	pxr::UsdShadeShader SurfaceShader = ShadeMaterial.ComputeSurfaceSource(MdlToken);
+
+	if (!SurfaceShader)
+	{
+		Super::CreateAssets();
+		return;
+	}
+
+	const FString MdlRootPath = Context->Stage.GetRootLayer().GetRealPath();
+
+	pxr::SdfAssetPath SurfaceSourceAssetPath;
+	SurfaceShader.GetSourceAsset(&SurfaceSourceAssetPath, MdlToken);
+
+	if (SurfaceSourceAssetPath.GetAssetPath().empty())
+	{
+		// Old Mdl Schema
+		const pxr::TfToken ModuleToken("module");
+		const pxr::UsdAttribute& MDLModule = SurfaceShader.GetPrim().GetAttribute(ModuleToken);
+
+		if (MDLModule.GetTypeName().GetAsToken() == pxr::SdfValueTypeNames->Asset)
+		{
+			SurfaceSourceAssetPath = UsdUtils::GetUsdValue<pxr::SdfAssetPath>(MDLModule, Context->Time);
+		}
+	}
+
+	const FString MdlOrigAssetPath = UsdToUnreal::ConvertString(SurfaceSourceAssetPath.GetAssetPath());
+	const FString MdlAbsoluteAssetPath = UsdToUnreal::ConvertString(SurfaceSourceAssetPath.GetResolvedPath());
+
+	const FString MdlModuleName = [&MdlAbsoluteAssetPath, &MdlRootPath]()
+	{
+		FString ModuleRelativePath = MdlAbsoluteAssetPath;
+		FPaths::MakePathRelativeTo(ModuleRelativePath, *MdlRootPath);
+
+		FString ModuleName = UE::Mdl::Util::ConvertFilePathToModuleName(*ModuleRelativePath);
+		return ModuleName;
+	}();
+
+	if (!MdlModuleName.IsEmpty() && Context->UsdAssetCache)
+	{
+		pxr::TfToken MdlDefinitionToken;
+		SurfaceShader.GetSourceAssetSubIdentifier(&MdlDefinitionToken, MdlToken);
+
+		const FString MdlDefinitionName = UsdToUnreal::ConvertToken(MdlDefinitionToken);
+
+		const FString MdlFullName = MdlModuleName + TEXT("::") + MdlDefinitionName;
+		const FString MdlFullInstanceName = MdlFullName + TEXT("_Instance");
+		const FString HashPrefix = UsdUtils::GetAssetHashPrefix(GetPrim(), Context->bShareAssetsForIdenticalPrims);
+		const FString MaterialHash = HashPrefix + MdlFullName;
+		const FString MdlSearchPath = FPaths::GetPath(Context->Stage.GetRootLayer().GetRealPath());
+
+		TSet<UObject*> Dependencies;
+
+		FScopedUnrealAllocs UEAllocs;
+
+		// Create reference material
+		//
+		// The FMdlMaterialImporter will (strangely) take the provided UPackage "ParentPackage" and instead just create
+		// another package, using "ParentPackage.GetName() / MdlMaterialName" as it's path... This is not ideal especially
+		// since it will lead to asset collisions if we're not reusing assets via hash. Here we just put everything
+		// on the transient package and rename the asset into the target destination package instead, like we do for MaterialX
+		//
+		// We have to do all of that inside the lambda though because we want to reuse that reference material if we have one
+		// cached already
+		bool bCreatedMaterial = false;
+		UMaterialInterface* MdlMaterial = Context->UsdAssetCache->GetOrCreateCustomCachedAsset<UMaterialInterface>(
+			MaterialHash,
+			MdlFullName,
+			Context->ObjectFlags,
+			[&MdlSearchPath, &MdlModuleName, &MdlDefinitionName, &MdlOrigAssetPath](UPackage* Outer, FName SanitizedName, EObjectFlags FlagsToUse)
+				-> UMaterialInterface*
+			{
+				// Add the USD root as a search path for MDL
+				FMdlMaterialImporter::FScopedSearchPath UsdDirMdlSearchPath(MdlSearchPath);
+
+				TStrongObjectPtr<UMDLImporterOptions> ImportOptions(NewObject<UMDLImporterOptions>());
+
+				UPackage* ParentPackage = GetTransientPackage();
+				UMaterialInterface* ReferenceMaterial = FMdlMaterialImporter::ImportMaterialFromModule(
+					ParentPackage,
+					FlagsToUse,
+					MdlModuleName,
+					MdlDefinitionName,
+					*ImportOptions.Get()
+				);
+				if (!ReferenceMaterial)
+				{
+					FString FullPath = FPaths::Combine(MdlSearchPath, MdlOrigAssetPath);
+					FPaths::NormalizeFilename(FullPath);
+					UE_LOG(LogUsdMdl, Warning, TEXT("Failed to load MDL material from file '%s'. Does the file exist?"), *FullPath);
+					return nullptr;
+				}
+
+				// Rename the UMaterialInterface into the target UPackage the asset cache created for us.
+				// SanitizedName will already match it.
+				const bool bRenamed = ReferenceMaterial->Rename(*SanitizedName.ToString(), Outer, REN_NonTransactional | REN_DontCreateRedirectors);
+				ensure(bRenamed);
+
+				// Let's not trust the flags the MDLImporter used and just use our own instead
+				ReferenceMaterial->ClearFlags(ReferenceMaterial->GetFlags());
+				ReferenceMaterial->SetFlags(FlagsToUse);
+
+				return ReferenceMaterial;
+			},
+			&bCreatedMaterial
+		);
+		if (bCreatedMaterial && MdlMaterial)
+		{
+			UE::MDLShadeTranslator::Private::NotifyIfMaterialNeedsVirtualTextures(MdlMaterial);
+			Dependencies.Append(IUsdClassesModule::GetAssetDependencies(MdlMaterial));
+		}
+		else if (!MdlMaterial)
+		{
+			FUsdLogManager::LogMessage(
+				EMessageSeverity::Warning,
+				FText::Format(
+					LOCTEXT("UsdMdlConversionFailed", "Failed to create MDL material for prim {0}."),
+					FText::FromString(PrimPath.GetString())
+				)
+			);
+		}
+
+		// Create material instance
+		bool bCreatedInstance = false;
+		const FString MaterialInstanceHash = HashPrefix + MdlFullInstanceName;
+		UMaterialInstanceConstant* MdlMaterialInstance = Context->UsdAssetCache->GetOrCreateCachedAsset<UMaterialInstanceConstant>(
+			MaterialInstanceHash,
+			MdlFullInstanceName,
+			Context->ObjectFlags,
+			&bCreatedInstance
+		);
+		if (bCreatedInstance && MdlMaterialInstance)
+		{
+			MdlMaterialInstance->SetParentEditorOnly(MdlMaterial);
+
+			UsdToUnreal::ConvertShadeInputsToParameters(
+				ShadeMaterial,
+				*MdlMaterialInstance,
+				Context->UsdAssetCache.Get(),
+				*Context->RenderContext.ToString(),
+				Context->bShareAssetsForIdenticalPrims
+			);
+
+			// We can't blindly recreate all component render states when a level is being added, because we may end up first creating
+			// render states for some components, and UWorld::AddToWorld calls FScene::AddPrimitive which expects the component to not have
+			// primitives yet
+			FMaterialUpdateContext::EOptions::Type Options = FMaterialUpdateContext::EOptions::Default;
+			if (Context->Level && Context->Level->bIsAssociatingLevel)
+			{
+				Options = (FMaterialUpdateContext::EOptions::Type)(Options & ~FMaterialUpdateContext::EOptions::RecreateRenderStates);
+			}
+
+			FMaterialUpdateContext UpdateContext(Options, GMaxRHIShaderPlatform);
+			UpdateContext.AddMaterialInstance(MdlMaterialInstance);
+			MdlMaterialInstance->PreEditChange(nullptr);
+			MdlMaterialInstance->PostEditChange();
+
+			UE::MDLShadeTranslator::Private::NotifyIfMaterialNeedsVirtualTextures(MdlMaterialInstance);
+			Dependencies.Append(IUsdClassesModule::GetAssetDependencies(MdlMaterialInstance));
+		}
+
+		// If we created a material here, we could have created some textures. Let's cache those too
+		// so that PostImportMaterial can take care of them
+		for (UObject* Object : Dependencies)
+		{
+			if (UTexture* Texture = Cast<UTexture>(Object))
+			{
+				// We have to watch out because it seems like FMdlMaterialImporter can generate some materials using
+				// default engine textures like FlatNormal. We don't want to cache those.
+				// Note that we can only test for that by checking whether the asset has been saved or not because the
+				// MDL importer (unfortunately) can't use our asset cache with its previously created materials and textures...
+				// This means that every time an MDL has to create a new texture it will re-read the image from disk again and create
+				// a new (as of yet unsaved) asset. This means if the asset is saved it must have been an engine texture.
+				// Also Note: I think the MDLImporter has a bug here where it will put those textures in paths like
+				// "/Engine/TextureName"... We can take them anyway though.
+				UPackage* TexturePackage = Texture->GetOutermost();
+				if (TexturePackage && TexturePackage->GetFileSize() > 0)
+				{
+					continue;
+				}
+
+				// Note how we add MaterialHash to the texture hash: This because unfortunately there's no way to get
+				// the MDLImporter to easily use our cached textures (it does have a ResourcesDir on UMDLImporterOptions but
+				// those seem to be file directory paths instead?). This means each MDL material must be allowed to create and
+				// generate new textures for itself, which we must guarantee don't collide with anything inside the asset cache.
+				const FString FilePath = Texture->AssetImportData ? Texture->AssetImportData->GetFirstFilename() : TEXT("");
+				const FString PrefixedTextureHash = HashPrefix + MaterialHash
+													+ UsdUtils::GetTextureHash(
+														FilePath,
+														Texture->SRGB,
+														Texture->CompressionSettings,
+														Texture->GetTextureAddressX(),
+														Texture->GetTextureAddressY()
+													);
+
+				bool bCreatedTexture = false;
+				UObject* PackagedTexture = Context->UsdAssetCache->GetOrCreateCustomCachedAsset(
+					PrefixedTextureHash,
+					Texture->GetClass(),
+					Texture->GetName(),
+					Texture->GetFlags(),
+					[Texture](UPackage* Outer, FName SanitizedName, EObjectFlags FlagsToUse) -> UObject*
+					{
+						// Rename the asset into the target UPackage the asset cache created for us.
+						// SanitizedName will already match it.
+						const bool bRenamed = Texture->Rename(*SanitizedName.ToString(), Outer, REN_NonTransactional | REN_DontCreateRedirectors);
+						ensure(bRenamed);
+
+						// Let's not trust the flags the MDLImporter used and just use our own instead
+						Texture->ClearFlags(Texture->GetFlags());
+						Texture->SetFlags(FlagsToUse);
+
+						return Texture;
+					},
+					&bCreatedTexture
+				);
+				ensure(PackagedTexture && bCreatedTexture);
+			}
+		}
+
+		PostImportMaterial(MaterialHash, MdlMaterial);
+		PostImportMaterial(MaterialInstanceHash, MdlMaterialInstance);
+	}
+	else
+	{
+		Super::CreateAssets();
+	}
+}
+
+#undef LOCTEXT_NAMESPACE
+
+#endif	  // #if USE_USD_SDK && WITH_EDITOR
diff --git a/Engine/Plugins/Importers/USDImporterMDL/Source/USDImporterMDL/Private/USDImporterMDLModule.cpp b/Engine/Plugins/Importers/USDImporterMDL/Source/USDImporterMDL/Private/USDImporterMDLModule.cpp
new file mode 100644
--- /dev/null
+++ b/Engine/Plugins/Importers/USDImporterMDL/Source/USDImporterMDL/Private/USDImporterMDLModule.cpp
@@ -0,0 +1,46 @@
+// Copyright Epic Games, Inc. All Rights Reserved.
+
+#include "MDLUSDLog.h"
+#include "MDLUSDShadeMaterialTranslator.h"
+#include "Objects/USDSchemaTranslator.h"
+#include "USDMaterialUtils.h"
+#include "USDMemory.h"
+
+#include "MDLImporterModule.h"
+
+DEFINE_LOG_CATEGORY(LogUsdMdl);
+
+class FUsdImporterMdlModule : public IModuleInterface
+{
+private:
+	virtual void StartupModule() override
+	{
+#if USE_USD_SDK && WITH_EDITOR
+#ifdef USE_MDLSDK
+		if (GIsEditor)
+		{
+			if (IMDLImporterModule* MDLImporterModule = FModuleManager::Get().LoadModulePtr<IMDLImporterModule>(TEXT("MDLImporter")))
+			{
+				UsdUnreal::MaterialUtils::RegisterRenderContext(UnrealIdentifiers::MdlRenderContext);
+
+				FUsdSchemaTranslatorRegistry& Registry = FUsdSchemaTranslatorRegistry::Get();
+				TranslatorHandle = Registry.Register<FMdlUsdShadeMaterialTranslator>(TEXT("UsdShadeMaterial"));
+			}
+		}
+#else
+		UE_LOG(LogUsdMdl, Log, TEXT("Not registering the MDL schema translator as the MDL SDK is not available"));
+#endif	  // USE_MDL_SDK
+#endif	  // WITH_EDITOR
+	}
+
+	virtual void ShutdownModule() override
+	{
+		FUsdSchemaTranslatorRegistry& Registry = FUsdSchemaTranslatorRegistry::Get();
+		Registry.Unregister(TranslatorHandle);
+	}
+
+private:
+	FRegisteredSchemaTranslatorHandle TranslatorHandle;
+};
+
+IMPLEMENT_MODULE_USD(FUsdImporterMdlModule, USDImporterMDL);
diff --git a/Engine/Plugins/Importers/USDImporterMDL/Source/USDImporterMDL/Public/MDLUSDLog.h b/Engine/Plugins/Importers/USDImporterMDL/Source/USDImporterMDL/Public/MDLUSDLog.h
new file mode 100644
--- /dev/null
+++ b/Engine/Plugins/Importers/USDImporterMDL/Source/USDImporterMDL/Public/MDLUSDLog.h
@@ -0,0 +1,7 @@
+// Copyright Epic Games, Inc. All Rights Reserved.
+
+#pragma once
+
+#include "Logging/LogMacros.h"
+
+USDIMPORTERMDL_API DECLARE_LOG_CATEGORY_EXTERN(LogUsdMdl, Log, All);
diff --git a/Engine/Plugins/Importers/USDImporterMDL/Source/USDImporterMDL/Public/MDLUSDShadeMaterialTranslator.h b/Engine/Plugins/Importers/USDImporterMDL/Source/USDImporterMDL/Public/MDLUSDShadeMaterialTranslator.h
new file mode 100644
--- /dev/null
+++ b/Engine/Plugins/Importers/USDImporterMDL/Source/USDImporterMDL/Public/MDLUSDShadeMaterialTranslator.h
@@ -0,0 +1,29 @@
+// Copyright Epic Games, Inc. All Rights Reserved.
+
+#pragma once
+
+#include "USDShadeMaterialTranslator.h"
+
+#include "Custom/MaterialXUSDShadeMaterialTranslator.h"
+
+#if USE_USD_SDK && WITH_EDITOR
+
+#include "USDIncludesStart.h"
+#include "pxr/pxr.h"
+#include "USDIncludesEnd.h"
+
+class FMdlUsdShadeMaterialTranslator : public FMaterialXUsdShadeMaterialTranslator
+{
+	using Super = FMaterialXUsdShadeMaterialTranslator;
+
+public:
+	UE_DEPRECATED(5.5, "Use UnrealIdentifiers::MdlRenderContext instead")
+	static FName MdlRenderContext;
+
+public:
+	using FMaterialXUsdShadeMaterialTranslator::FMaterialXUsdShadeMaterialTranslator;
+
+	virtual void CreateAssets() override;
+};
+
+#endif	  // #if USE_USD_SDK && WITH_EDITOR
diff --git a/Engine/Plugins/Importers/USDImporterMDL/Source/USDImporterMDL/USDImporterMDL.Build.cs b/Engine/Plugins/Importers/USDImporterMDL/Source/USDImporterMDL/USDImporterMDL.Build.cs
new file mode 100644
--- /dev/null
+++ b/Engine/Plugins/Importers/USDImporterMDL/Source/USDImporterMDL/USDImporterMDL.Build.cs
@@ -0,0 +1,45 @@
+// Copyright Epic Games, Inc. All Rights Reserved.
+
+using System;
+using System.IO;
+using UnrealBuildTool;
+
+public class USDImporterMDL : ModuleRules
+{
+	public USDImporterMDL(ReadOnlyTargetRules Target) : base(Target)
+	{
+		bUseRTTI = true;
+
+		PublicDependencyModuleNames.AddRange(
+			new string[] {
+				"Core",
+				"CoreUObject",
+			}
+		);
+
+		PrivateDependencyModuleNames.AddRange(
+			new string[] {
+				"Engine",
+				"MDLImporter", // Always add the dependency so we don't have to #ifdef our code that uses the MDLImporter includes directly
+				"RHI", // Because the MDL shader translator uses GMaxRHIShaderPlatform
+				"UnrealUSDWrapper", // Needed to link with USD and its dependencies, like Python
+				"USDClasses", // To register the render context
+				"USDSchemas", // Because the MDL schema translator actually derives the MDL schema translator
+				"USDUtilities", // To register the new schema translator
+			}
+		);
+
+		PrivateDefinitions.Add("SUPPRESS_PER_MODULE_INLINE_FILE");
+
+		if (Target.bBuildEditor)
+		{
+			// Note: This path must match whatever the MDLImporter.Build.cs is actually using
+			string ExpectedMDLSDKPath = Path.GetFullPath(Path.Combine(EngineDirectory, "Restricted/NotForLicensees/Source/ThirdParty/Enterprise/mdl-sdk-349500.8766a/include"));
+			if (Directory.Exists(ExpectedMDLSDKPath))
+			{
+				// Conditionally adding this define lets us not register the schema translator if the SDK is not usable anyway
+				PrivateDefinitions.Add("USE_MDLSDK");
+			}
+		}
+	}
+}
diff --git a/Engine/Plugins/Importers/USDImporterMDL/USDImporterMDL.uplugin b/Engine/Plugins/Importers/USDImporterMDL/USDImporterMDL.uplugin
new file mode 100644
--- /dev/null
+++ b/Engine/Plugins/Importers/USDImporterMDL/USDImporterMDL.uplugin
@@ -0,0 +1,43 @@
+{
+	"FileVersion": 3,
+	"Version": 1,
+	"VersionName": "1.0",
+	"FriendlyName": "USD Importer MDL integration",
+	"Description": "Allows importing USD files that reference MDL files, via the USD Stage Actor and USD import",
+	"Category": "Importers",
+	"CreatedBy": "Epic Games, Inc.",
+	"CreatedByURL": "https://epicgames.com",
+	"DocsURL": "",
+	"MarketplaceURL": "",
+	"SupportURL": "",
+	"EnabledByDefault": false,
+	"CanContainContent": false,
+	"IsBetaVersion": true,
+	"Installed": false,
+	"SupportedTargetPlatforms": [
+		"Win64",
+		"Mac",
+		"Linux"
+	],
+	"Modules": [
+		{
+			"Name": "USDImporterMDL",
+			"Type": "Editor",
+			"LoadingPhase": "Default"
+		}
+	],
+	"Plugins": [
+		{
+			"Name": "USDCore",
+			"Enabled": true
+		},
+		{
+			"Name": "USDImporter",
+			"Enabled": true
+		},
+		{
+			"Name": "MDLImporter",
+			"Enabled": true
+		}
+	]
+}
\ No newline at end of file
