USDImporterMDL: Fix mismatched supported target platforms.

diff --git a/Engine/Plugins/Importers/USDImporterMDL/USDImporterMDL.uplugin b/Engine/Plugins/Importers/USDImporterMDL/USDImporterMDL.uplugin
--- a/Engine/Plugins/Importers/USDImporterMDL/USDImporterMDL.uplugin
+++ b/Engine/Plugins/Importers/USDImporterMDL/USDImporterMDL.uplugin
@@ -23,7 +23,10 @@
 		{
 			"Name": "USDImporterMDL",
 			"Type": "Editor",
-			"LoadingPhase": "Default"
+			"LoadingPhase": "Default",
+			"PlatformAllowList": [
+				"Win64", "Mac", "Linux"
+			]
 		}
 	],
 	"Plugins": [
