USD: Fixed potential build issue when the USD SDK is disabled.

diff --git a/Engine/Plugins/Experimental/ChaosCachingUSD/Source/ChaosCachingUSD/ChaosCachingUSD.build.cs b/Engine/Plugins/Experimental/ChaosCachingUSD/Source/ChaosCachingUSD/ChaosCachingUSD.build.cs
--- a/Engine/Plugins/Experimental/ChaosCachingUSD/Source/ChaosCachingUSD/ChaosCachingUSD.build.cs
+++ b/Engine/Plugins/Experimental/ChaosCachingUSD/Source/ChaosCachingUSD/ChaosCachingUSD.build.cs
@@ -32,11 +32,36 @@ namespace UnrealBuildTool.Rules
 			);
 
 			bool bUseUsdMemoryOverloadsIfSupported = true;
-			bool bEnableUsdMemoryOverloads = bUseUsdMemoryOverloadsIfSupported && Target.Platform == UnrealTargetPlatform.Win64 && Target.Architecture == UnrealArch.X64;
+			bool bEnableUsdMemoryOverloads = EnableUsdSdk(Target) && bUseUsdMemoryOverloadsIfSupported && Target.Platform == UnrealTargetPlatform.Win64 && Target.Architecture == UnrealArch.X64;
 			if (!bEnableUsdMemoryOverloads)
 			{
 				PrivateDefinitions.Add("SUPPRESS_PER_MODULE_INLINE_FILE"); // This module does not use core's standard operator new/delete overloads
 			}
 		}
+
+		bool EnableUsdSdk(ReadOnlyTargetRules Target)
+		{
+			bool bEnableUsdSdk = (
+				Target.WindowsPlatform.Compiler != WindowsCompiler.Intel &&
+				Target.StaticAnalyzer == StaticAnalyzer.None
+			);
+
+			// Don't enable USD when running the include tool because it has issues parsing Boost headers
+			if (Target.GlobalDefinitions.Contains("UE_INCLUDE_TOOL=1"))
+			{
+				bEnableUsdSdk = false;
+			}
+
+			bool bIsGameOrMonolithic = Target.LinkType == TargetLinkType.Monolithic || Target.Type == TargetType.Client || Target.Type == TargetType.Server || Target.Type == TargetType.Game;
+			if (bEnableUsdSdk &&
+				bIsGameOrMonolithic &&
+				((Target.StaticAllocator != StaticAllocatorType.Ansi) || !Target.GlobalDefinitions.Contains("UE_USE_MALLOC_FILL_BYTES=0"))
+			)
+			{
+				bEnableUsdSdk = false;
+			}
+
+			return bEnableUsdSdk;
+		}
 	}
 }
\ No newline at end of file
diff --git a/Engine/Plugins/Experimental/ChaosFlesh/Source/ChaosFleshEngine/ChaosFleshEngine.Build.cs b/Engine/Plugins/Experimental/ChaosFlesh/Source/ChaosFleshEngine/ChaosFleshEngine.Build.cs
--- a/Engine/Plugins/Experimental/ChaosFlesh/Source/ChaosFleshEngine/ChaosFleshEngine.Build.cs
+++ b/Engine/Plugins/Experimental/ChaosFlesh/Source/ChaosFleshEngine/ChaosFleshEngine.Build.cs
@@ -89,7 +89,7 @@ namespace UnrealBuildTool.Rules
 				PublicDefinitions.Add("DO_USD_CACHING=1");
 
 				bool bUseUsdMemoryOverloadsIfSupported = true;
-				bool bEnableUsdMemoryOverloads = bUseUsdMemoryOverloadsIfSupported && Target.Platform == UnrealTargetPlatform.Win64 && Target.Architecture == UnrealArch.X64;
+				bool bEnableUsdMemoryOverloads = EnableUsdSdk(Target) && bUseUsdMemoryOverloadsIfSupported && Target.Platform == UnrealTargetPlatform.Win64 && Target.Architecture == UnrealArch.X64;
 				if (!bEnableUsdMemoryOverloads)
 				{
 					PrivateDefinitions.Add("SUPPRESS_PER_MODULE_INLINE_FILE"); // This module does not use core's standard operator new/delete overloads
diff --git a/Engine/Plugins/Importers/USDImporter/Source/USDExporter/USDExporter.Build.cs b/Engine/Plugins/Importers/USDImporter/Source/USDExporter/USDExporter.Build.cs
--- a/Engine/Plugins/Importers/USDImporter/Source/USDExporter/USDExporter.Build.cs
+++ b/Engine/Plugins/Importers/USDImporter/Source/USDExporter/USDExporter.Build.cs
@@ -2,6 +2,7 @@
 
 using System;
 using System.IO;
+using System.Linq;
 
 namespace UnrealBuildTool.Rules
 {
@@ -61,11 +62,36 @@ namespace UnrealBuildTool.Rules
 			);
 
 			bool bUseUsdMemoryOverloadsIfSupported = true;
-			bool bEnableUsdMemoryOverloads = bUseUsdMemoryOverloadsIfSupported && Target.Platform == UnrealTargetPlatform.Win64 && Target.Architecture == UnrealArch.X64;
+			bool bEnableUsdMemoryOverloads = EnableUsdSdk(Target) && bUseUsdMemoryOverloadsIfSupported && Target.Platform == UnrealTargetPlatform.Win64 && Target.Architecture == UnrealArch.X64;
 			if (!bEnableUsdMemoryOverloads)
 			{
 				PrivateDefinitions.Add("SUPPRESS_PER_MODULE_INLINE_FILE"); // This module does not use core's standard operator new/delete overloads
 			}
 		}
+
+		bool EnableUsdSdk(ReadOnlyTargetRules Target)
+		{
+			bool bEnableUsdSdk = (
+				Target.WindowsPlatform.Compiler != WindowsCompiler.Intel &&
+				Target.StaticAnalyzer == StaticAnalyzer.None
+			);
+
+			// Don't enable USD when running the include tool because it has issues parsing Boost headers
+			if (Target.GlobalDefinitions.Contains("UE_INCLUDE_TOOL=1"))
+			{
+				bEnableUsdSdk = false;
+			}
+
+			bool bIsGameOrMonolithic = Target.LinkType == TargetLinkType.Monolithic || Target.Type == TargetType.Client || Target.Type == TargetType.Server || Target.Type == TargetType.Game;
+			if (bEnableUsdSdk &&
+				bIsGameOrMonolithic &&
+				((Target.StaticAllocator != StaticAllocatorType.Ansi) || !Target.GlobalDefinitions.Contains("UE_USE_MALLOC_FILL_BYTES=0"))
+			)
+			{
+				bEnableUsdSdk = false;
+			}
+
+			return bEnableUsdSdk;
+		}
 	}
 }
diff --git a/Engine/Plugins/Importers/USDImporter/Source/USDSchemas/USDSchemas.Build.cs b/Engine/Plugins/Importers/USDImporter/Source/USDSchemas/USDSchemas.Build.cs
--- a/Engine/Plugins/Importers/USDImporter/Source/USDSchemas/USDSchemas.Build.cs
+++ b/Engine/Plugins/Importers/USDImporter/Source/USDSchemas/USDSchemas.Build.cs
@@ -2,6 +2,7 @@
 
 using System;
 using System.IO;
+using System.Linq;
 
 namespace UnrealBuildTool.Rules
 {
@@ -80,11 +81,36 @@ namespace UnrealBuildTool.Rules
 			}
 
 			bool bUseUsdMemoryOverloadsIfSupported = true;
-			bool bEnableUsdMemoryOverloads = bUseUsdMemoryOverloadsIfSupported && Target.Platform == UnrealTargetPlatform.Win64 && Target.Architecture == UnrealArch.X64;
+			bool bEnableUsdMemoryOverloads = EnableUsdSdk(Target) && bUseUsdMemoryOverloadsIfSupported && Target.Platform == UnrealTargetPlatform.Win64 && Target.Architecture == UnrealArch.X64;
 			if (!bEnableUsdMemoryOverloads)
 			{
 				PrivateDefinitions.Add("SUPPRESS_PER_MODULE_INLINE_FILE"); // This module does not use core's standard operator new/delete overloads
 			}
 		}
+
+		bool EnableUsdSdk(ReadOnlyTargetRules Target)
+		{
+			bool bEnableUsdSdk = (
+				Target.WindowsPlatform.Compiler != WindowsCompiler.Intel &&
+				Target.StaticAnalyzer == StaticAnalyzer.None
+			);
+
+			// Don't enable USD when running the include tool because it has issues parsing Boost headers
+			if (Target.GlobalDefinitions.Contains("UE_INCLUDE_TOOL=1"))
+			{
+				bEnableUsdSdk = false;
+			}
+
+			bool bIsGameOrMonolithic = Target.LinkType == TargetLinkType.Monolithic || Target.Type == TargetType.Client || Target.Type == TargetType.Server || Target.Type == TargetType.Game;
+			if (bEnableUsdSdk &&
+				bIsGameOrMonolithic &&
+				((Target.StaticAllocator != StaticAllocatorType.Ansi) || !Target.GlobalDefinitions.Contains("UE_USE_MALLOC_FILL_BYTES=0"))
+			)
+			{
+				bEnableUsdSdk = false;
+			}
+
+			return bEnableUsdSdk;
+		}
 	}
 }
diff --git a/Engine/Plugins/Importers/USDImporter/Source/USDStage/USDStage.Build.cs b/Engine/Plugins/Importers/USDImporter/Source/USDStage/USDStage.Build.cs
--- a/Engine/Plugins/Importers/USDImporter/Source/USDStage/USDStage.Build.cs
+++ b/Engine/Plugins/Importers/USDImporter/Source/USDStage/USDStage.Build.cs
@@ -2,6 +2,7 @@
 
 using System;
 using System.IO;
+using System.Linq;
 
 namespace UnrealBuildTool.Rules
 {
@@ -72,11 +73,36 @@ namespace UnrealBuildTool.Rules
 			}
 
 			bool bUseUsdMemoryOverloadsIfSupported = true;
-			bool bEnableUsdMemoryOverloads = bUseUsdMemoryOverloadsIfSupported && Target.Platform == UnrealTargetPlatform.Win64 && Target.Architecture == UnrealArch.X64;
+			bool bEnableUsdMemoryOverloads = EnableUsdSdk(Target) && bUseUsdMemoryOverloadsIfSupported && Target.Platform == UnrealTargetPlatform.Win64 && Target.Architecture == UnrealArch.X64;
 			if (!bEnableUsdMemoryOverloads)
 			{
 				PrivateDefinitions.Add("SUPPRESS_PER_MODULE_INLINE_FILE"); // This module does not use core's standard operator new/delete overloads
 			}
 		}
+
+		bool EnableUsdSdk(ReadOnlyTargetRules Target)
+		{
+			bool bEnableUsdSdk = (
+				Target.WindowsPlatform.Compiler != WindowsCompiler.Intel &&
+				Target.StaticAnalyzer == StaticAnalyzer.None
+			);
+
+			// Don't enable USD when running the include tool because it has issues parsing Boost headers
+			if (Target.GlobalDefinitions.Contains("UE_INCLUDE_TOOL=1"))
+			{
+				bEnableUsdSdk = false;
+			}
+
+			bool bIsGameOrMonolithic = Target.LinkType == TargetLinkType.Monolithic || Target.Type == TargetType.Client || Target.Type == TargetType.Server || Target.Type == TargetType.Game;
+			if (bEnableUsdSdk &&
+				bIsGameOrMonolithic &&
+				((Target.StaticAllocator != StaticAllocatorType.Ansi) || !Target.GlobalDefinitions.Contains("UE_USE_MALLOC_FILL_BYTES=0"))
+			)
+			{
+				bEnableUsdSdk = false;
+			}
+
+			return bEnableUsdSdk;
+		}
 	}
 }
diff --git a/Engine/Plugins/Importers/USDImporter/Source/USDStageEditor/USDStageEditor.Build.cs b/Engine/Plugins/Importers/USDImporter/Source/USDStageEditor/USDStageEditor.Build.cs
--- a/Engine/Plugins/Importers/USDImporter/Source/USDStageEditor/USDStageEditor.Build.cs
+++ b/Engine/Plugins/Importers/USDImporter/Source/USDStageEditor/USDStageEditor.Build.cs
@@ -2,6 +2,7 @@
 
 using System;
 using System.IO;
+using System.Linq;
 
 namespace UnrealBuildTool.Rules
 {
@@ -40,11 +41,36 @@ namespace UnrealBuildTool.Rules
 			);
 
 			bool bUseUsdMemoryOverloadsIfSupported = true;
-			bool bEnableUsdMemoryOverloads = bUseUsdMemoryOverloadsIfSupported && Target.Platform == UnrealTargetPlatform.Win64 && Target.Architecture == UnrealArch.X64;
+			bool bEnableUsdMemoryOverloads = EnableUsdSdk(Target) && bUseUsdMemoryOverloadsIfSupported && Target.Platform == UnrealTargetPlatform.Win64 && Target.Architecture == UnrealArch.X64;
 			if (!bEnableUsdMemoryOverloads)
 			{
 				PrivateDefinitions.Add("SUPPRESS_PER_MODULE_INLINE_FILE"); // This module does not use core's standard operator new/delete overloads
 			}
 		}
+
+		bool EnableUsdSdk(ReadOnlyTargetRules Target)
+		{
+			bool bEnableUsdSdk = (
+				Target.WindowsPlatform.Compiler != WindowsCompiler.Intel &&
+				Target.StaticAnalyzer == StaticAnalyzer.None
+			);
+
+			// Don't enable USD when running the include tool because it has issues parsing Boost headers
+			if (Target.GlobalDefinitions.Contains("UE_INCLUDE_TOOL=1"))
+			{
+				bEnableUsdSdk = false;
+			}
+
+			bool bIsGameOrMonolithic = Target.LinkType == TargetLinkType.Monolithic || Target.Type == TargetType.Client || Target.Type == TargetType.Server || Target.Type == TargetType.Game;
+			if (bEnableUsdSdk &&
+				bIsGameOrMonolithic &&
+				((Target.StaticAllocator != StaticAllocatorType.Ansi) || !Target.GlobalDefinitions.Contains("UE_USE_MALLOC_FILL_BYTES=0"))
+			)
+			{
+				bEnableUsdSdk = false;
+			}
+
+			return bEnableUsdSdk;
+		}
 	}
 }
diff --git a/Engine/Plugins/Importers/USDImporter/Source/USDStageEditorViewModels/USDStageEditorViewModels.Build.cs b/Engine/Plugins/Importers/USDImporter/Source/USDStageEditorViewModels/USDStageEditorViewModels.Build.cs
--- a/Engine/Plugins/Importers/USDImporter/Source/USDStageEditorViewModels/USDStageEditorViewModels.Build.cs
+++ b/Engine/Plugins/Importers/USDImporter/Source/USDStageEditorViewModels/USDStageEditorViewModels.Build.cs
@@ -2,6 +2,7 @@
 
 using System;
 using System.IO;
+using System.Linq;
 
 namespace UnrealBuildTool.Rules
 {
@@ -34,11 +35,36 @@ namespace UnrealBuildTool.Rules
 			);
 
 			bool bUseUsdMemoryOverloadsIfSupported = true;
-			bool bEnableUsdMemoryOverloads = bUseUsdMemoryOverloadsIfSupported && Target.Platform == UnrealTargetPlatform.Win64 && Target.Architecture == UnrealArch.X64;
+			bool bEnableUsdMemoryOverloads = EnableUsdSdk(Target) && bUseUsdMemoryOverloadsIfSupported && Target.Platform == UnrealTargetPlatform.Win64 && Target.Architecture == UnrealArch.X64;
 			if (!bEnableUsdMemoryOverloads)
 			{
 				PrivateDefinitions.Add("SUPPRESS_PER_MODULE_INLINE_FILE"); // This module does not use core's standard operator new/delete overloads
 			}
 		}
+
+		bool EnableUsdSdk(ReadOnlyTargetRules Target)
+		{
+			bool bEnableUsdSdk = (
+				Target.WindowsPlatform.Compiler != WindowsCompiler.Intel &&
+				Target.StaticAnalyzer == StaticAnalyzer.None
+			);
+
+			// Don't enable USD when running the include tool because it has issues parsing Boost headers
+			if (Target.GlobalDefinitions.Contains("UE_INCLUDE_TOOL=1"))
+			{
+				bEnableUsdSdk = false;
+			}
+
+			bool bIsGameOrMonolithic = Target.LinkType == TargetLinkType.Monolithic || Target.Type == TargetType.Client || Target.Type == TargetType.Server || Target.Type == TargetType.Game;
+			if (bEnableUsdSdk &&
+				bIsGameOrMonolithic &&
+				((Target.StaticAllocator != StaticAllocatorType.Ansi) || !Target.GlobalDefinitions.Contains("UE_USE_MALLOC_FILL_BYTES=0"))
+			)
+			{
+				bEnableUsdSdk = false;
+			}
+
+			return bEnableUsdSdk;
+		}
 	}
 }
diff --git a/Engine/Plugins/Importers/USDImporter/Source/USDStageImporter/USDStageImporter.Build.cs b/Engine/Plugins/Importers/USDImporter/Source/USDStageImporter/USDStageImporter.Build.cs
--- a/Engine/Plugins/Importers/USDImporter/Source/USDStageImporter/USDStageImporter.Build.cs
+++ b/Engine/Plugins/Importers/USDImporter/Source/USDStageImporter/USDStageImporter.Build.cs
@@ -2,6 +2,7 @@
 
 using System;
 using System.IO;
+using System.Linq;
 
 namespace UnrealBuildTool.Rules
 {
@@ -46,11 +47,36 @@ namespace UnrealBuildTool.Rules
 			);
 
 			bool bUseUsdMemoryOverloadsIfSupported = true;
-			bool bEnableUsdMemoryOverloads = bUseUsdMemoryOverloadsIfSupported && Target.Platform == UnrealTargetPlatform.Win64 && Target.Architecture == UnrealArch.X64;
+			bool bEnableUsdMemoryOverloads = EnableUsdSdk(Target) && bUseUsdMemoryOverloadsIfSupported && Target.Platform == UnrealTargetPlatform.Win64 && Target.Architecture == UnrealArch.X64;
 			if (!bEnableUsdMemoryOverloads)
 			{
 				PrivateDefinitions.Add("SUPPRESS_PER_MODULE_INLINE_FILE"); // This module does not use core's standard operator new/delete overloads
 			}
 		}
+
+		bool EnableUsdSdk(ReadOnlyTargetRules Target)
+		{
+			bool bEnableUsdSdk = (
+				Target.WindowsPlatform.Compiler != WindowsCompiler.Intel &&
+				Target.StaticAnalyzer == StaticAnalyzer.None
+			);
+
+			// Don't enable USD when running the include tool because it has issues parsing Boost headers
+			if (Target.GlobalDefinitions.Contains("UE_INCLUDE_TOOL=1"))
+			{
+				bEnableUsdSdk = false;
+			}
+
+			bool bIsGameOrMonolithic = Target.LinkType == TargetLinkType.Monolithic || Target.Type == TargetType.Client || Target.Type == TargetType.Server || Target.Type == TargetType.Game;
+			if (bEnableUsdSdk &&
+				bIsGameOrMonolithic &&
+				((Target.StaticAllocator != StaticAllocatorType.Ansi) || !Target.GlobalDefinitions.Contains("UE_USE_MALLOC_FILL_BYTES=0"))
+			)
+			{
+				bEnableUsdSdk = false;
+			}
+
+			return bEnableUsdSdk;
+		}
 	}
 }
diff --git a/Engine/Plugins/Importers/USDImporter/Source/USDTests/USDTests.Build.cs b/Engine/Plugins/Importers/USDImporter/Source/USDTests/USDTests.Build.cs
--- a/Engine/Plugins/Importers/USDImporter/Source/USDTests/USDTests.Build.cs
+++ b/Engine/Plugins/Importers/USDImporter/Source/USDTests/USDTests.Build.cs
@@ -2,6 +2,7 @@
 
 using System;
 using System.IO;
+using System.Linq;
 
 namespace UnrealBuildTool.Rules
 {
@@ -35,11 +36,36 @@ namespace UnrealBuildTool.Rules
 			}
 
 			bool bUseUsdMemoryOverloadsIfSupported = true;
-			bool bEnableUsdMemoryOverloads = bUseUsdMemoryOverloadsIfSupported && Target.Platform == UnrealTargetPlatform.Win64 && Target.Architecture == UnrealArch.X64;
+			bool bEnableUsdMemoryOverloads = EnableUsdSdk(Target) && bUseUsdMemoryOverloadsIfSupported && Target.Platform == UnrealTargetPlatform.Win64 && Target.Architecture == UnrealArch.X64;
 			if (!bEnableUsdMemoryOverloads)
 			{
 				PrivateDefinitions.Add("SUPPRESS_PER_MODULE_INLINE_FILE"); // This module does not use core's standard operator new/delete overloads
 			}
 		}
+
+		bool EnableUsdSdk(ReadOnlyTargetRules Target)
+		{
+			bool bEnableUsdSdk = (
+				Target.WindowsPlatform.Compiler != WindowsCompiler.Intel &&
+				Target.StaticAnalyzer == StaticAnalyzer.None
+			);
+
+			// Don't enable USD when running the include tool because it has issues parsing Boost headers
+			if (Target.GlobalDefinitions.Contains("UE_INCLUDE_TOOL=1"))
+			{
+				bEnableUsdSdk = false;
+			}
+
+			bool bIsGameOrMonolithic = Target.LinkType == TargetLinkType.Monolithic || Target.Type == TargetType.Client || Target.Type == TargetType.Server || Target.Type == TargetType.Game;
+			if (bEnableUsdSdk &&
+				bIsGameOrMonolithic &&
+				((Target.StaticAllocator != StaticAllocatorType.Ansi) || !Target.GlobalDefinitions.Contains("UE_USE_MALLOC_FILL_BYTES=0"))
+			)
+			{
+				bEnableUsdSdk = false;
+			}
+
+			return bEnableUsdSdk;
+		}
 	}
 }
diff --git a/Engine/Plugins/Importers/USDImporterMDL/Source/USDImporterMDL/USDImporterMDL.Build.cs b/Engine/Plugins/Importers/USDImporterMDL/Source/USDImporterMDL/USDImporterMDL.Build.cs
--- a/Engine/Plugins/Importers/USDImporterMDL/Source/USDImporterMDL/USDImporterMDL.Build.cs
+++ b/Engine/Plugins/Importers/USDImporterMDL/Source/USDImporterMDL/USDImporterMDL.Build.cs
@@ -2,6 +2,7 @@
 
 using System;
 using System.IO;
+using System.Linq;
 using UnrealBuildTool;
 
 public class USDImporterMDL : ModuleRules
@@ -30,7 +31,7 @@ public class USDImporterMDL : ModuleRules
 		);
 
 		bool bUseUsdMemoryOverloadsIfSupported = true;
-		bool bEnableUsdMemoryOverloads = bUseUsdMemoryOverloadsIfSupported && Target.Platform == UnrealTargetPlatform.Win64 && Target.Architecture == UnrealArch.X64;
+		bool bEnableUsdMemoryOverloads = EnableUsdSdk(Target) && bUseUsdMemoryOverloadsIfSupported && Target.Platform == UnrealTargetPlatform.Win64 && Target.Architecture == UnrealArch.X64;
 		if (!bEnableUsdMemoryOverloads)
 		{
 			PrivateDefinitions.Add("SUPPRESS_PER_MODULE_INLINE_FILE"); // This module does not use core's standard operator new/delete overloads
@@ -47,4 +48,28 @@ public class USDImporterMDL : ModuleRules
 			}
 		}
 	}
+	bool EnableUsdSdk(ReadOnlyTargetRules Target)
+	{
+		bool bEnableUsdSdk = (
+			Target.WindowsPlatform.Compiler != WindowsCompiler.Intel &&
+			Target.StaticAnalyzer == StaticAnalyzer.None
+		);
+
+		// Don't enable USD when running the include tool because it has issues parsing Boost headers
+		if (Target.GlobalDefinitions.Contains("UE_INCLUDE_TOOL=1"))
+		{
+			bEnableUsdSdk = false;
+		}
+
+		bool bIsGameOrMonolithic = Target.LinkType == TargetLinkType.Monolithic || Target.Type == TargetType.Client || Target.Type == TargetType.Server || Target.Type == TargetType.Game;
+		if (bEnableUsdSdk &&
+			bIsGameOrMonolithic &&
+			((Target.StaticAllocator != StaticAllocatorType.Ansi) || !Target.GlobalDefinitions.Contains("UE_USE_MALLOC_FILL_BYTES=0"))
+		)
+		{
+			bEnableUsdSdk = false;
+		}
+
+		return bEnableUsdSdk;
+	}
 }
diff --git a/Engine/Plugins/Runtime/USDCore/Source/USDUtilities/USDUtilities.Build.cs b/Engine/Plugins/Runtime/USDCore/Source/USDUtilities/USDUtilities.Build.cs
--- a/Engine/Plugins/Runtime/USDCore/Source/USDUtilities/USDUtilities.Build.cs
+++ b/Engine/Plugins/Runtime/USDCore/Source/USDUtilities/USDUtilities.Build.cs
@@ -2,6 +2,7 @@
 
 using System;
 using System.IO;
+using System.Linq;
 
 namespace UnrealBuildTool.Rules
 {
@@ -85,11 +86,36 @@ namespace UnrealBuildTool.Rules
 			PublicDefinitions.Add("ENABLE_USD_MATERIALX=" + (bEnableUsdMaterialX ? "1" : "0"));
 
 			bool bUseUsdMemoryOverloadsIfSupported = true;
-			bool bEnableUsdMemoryOverloads = bUseUsdMemoryOverloadsIfSupported && Target.Platform == UnrealTargetPlatform.Win64 && Target.Architecture == UnrealArch.X64;
+			bool bEnableUsdMemoryOverloads = EnableUsdSdk(Target) && bUseUsdMemoryOverloadsIfSupported && Target.Platform == UnrealTargetPlatform.Win64 && Target.Architecture == UnrealArch.X64;
 			if (!bEnableUsdMemoryOverloads)
 			{
 				PrivateDefinitions.Add("SUPPRESS_PER_MODULE_INLINE_FILE"); // This module does not use core's standard operator new/delete overloads
 			}
 		}
+
+		bool EnableUsdSdk(ReadOnlyTargetRules Target)
+		{
+			bool bEnableUsdSdk = (
+				Target.WindowsPlatform.Compiler != WindowsCompiler.Intel &&
+				Target.StaticAnalyzer == StaticAnalyzer.None
+			);
+
+			// Don't enable USD when running the include tool because it has issues parsing Boost headers
+			if (Target.GlobalDefinitions.Contains("UE_INCLUDE_TOOL=1"))
+			{
+				bEnableUsdSdk = false;
+			}
+
+			bool bIsGameOrMonolithic = Target.LinkType == TargetLinkType.Monolithic || Target.Type == TargetType.Client || Target.Type == TargetType.Server || Target.Type == TargetType.Game;
+			if (bEnableUsdSdk &&
+				bIsGameOrMonolithic &&
+				((Target.StaticAllocator != StaticAllocatorType.Ansi) || !Target.GlobalDefinitions.Contains("UE_USE_MALLOC_FILL_BYTES=0"))
+			)
+			{
+				bEnableUsdSdk = false;
+			}
+
+			return bEnableUsdSdk;
+		}
 	}
 }
diff --git a/Engine/Plugins/Runtime/USDCore/Source/UnrealUSDWrapper/UnrealUSDWrapper.Build.cs b/Engine/Plugins/Runtime/USDCore/Source/UnrealUSDWrapper/UnrealUSDWrapper.Build.cs
--- a/Engine/Plugins/Runtime/USDCore/Source/UnrealUSDWrapper/UnrealUSDWrapper.Build.cs
+++ b/Engine/Plugins/Runtime/USDCore/Source/UnrealUSDWrapper/UnrealUSDWrapper.Build.cs
@@ -25,7 +25,7 @@ namespace UnrealBuildTool.Rules
 			);
 
 			bool bUseUsdMemoryOverloadsIfSupported = true;
-			bool bEnableUsdMemoryOverloads = bUseUsdMemoryOverloadsIfSupported && Target.Platform == UnrealTargetPlatform.Win64 && Target.Architecture == UnrealArch.X64;
+			bool bEnableUsdMemoryOverloads = EnableUsdSdk(Target) && bUseUsdMemoryOverloadsIfSupported && Target.Platform == UnrealTargetPlatform.Win64 && Target.Architecture == UnrealArch.X64;
 			PublicDefinitions.Add("ENABLE_USD_MEMORY_OVERLOADS=" + (bEnableUsdMemoryOverloads ? "1" : "0"));
 
 			// Temporarily disabled runtime USD support until Mac and Linux dynamic linking issues are resolved
